Motion Trail Editor: Change to other Init() function which is called when the toolkit is created. The other Init() function wasn't getting called, which caused a crash because TimingStatsTextWidget wasn't initialized

diff --git a/Engine/Plugins/Experimental/MotionTrailEditorMode/Source/MotionTrailEditorMode/Private/MotionTrailEditorModeToolkit.cpp b/Engine/Plugins/Experimental/MotionTrailEditorMode/Source/MotionTrailEditorMode/Private/MotionTrailEditorModeToolkit.cpp
--- a/Engine/Plugins/Experimental/MotionTrailEditorMode/Source/MotionTrailEditorMode/Private/MotionTrailEditorModeToolkit.cpp
+++ b/Engine/Plugins/Experimental/MotionTrailEditorMode/Source/MotionTrailEditorMode/Private/MotionTrailEditorModeToolkit.cpp
@@ -24,10 +24,11 @@ FMotionTrailEditorModeToolkit::FMotionTrailEditorModeToolkit()
 {
 }
 
-void FMotionTrailEditorModeToolkit::Init(const TSharedPtr<IToolkitHost>& InitToolkitHost)
+void FMotionTrailEditorModeToolkit::Init(const TSharedPtr<IToolkitHost>& InitToolkitHost, TWeakObjectPtr<UEdMode> InOwningMode)
 {
 	SAssignNew(TimingStatsTextWidget, STextBlock);
-	FModeToolkit::Init(InitToolkitHost);
+
+	FModeToolkit::Init(InitToolkitHost, InOwningMode);
 }
 
 FName FMotionTrailEditorModeToolkit::GetToolkitFName() const
@@ -42,7 +43,6 @@ FText FMotionTrailEditorModeToolkit::GetBaseToolkitName() const
 
 TSharedPtr<SWidget> FMotionTrailEditorModeToolkit::GetInlineContent() const
 {
-	// TODO: fix crash where TimingStatsTextWidget becomes null when re-activating ed mode
 	TSharedPtr<SWidget> ModeWidget = FModeToolkit::GetInlineContent();
 	return SNew(SVerticalBox)
 		+ SVerticalBox::Slot()
diff --git a/Engine/Plugins/Experimental/MotionTrailEditorMode/Source/MotionTrailEditorMode/Public/MotionTrailEditorMode.h b/Engine/Plugins/Experimental/MotionTrailEditorMode/Source/MotionTrailEditorMode/Public/MotionTrailEditorMode.h
--- a/Engine/Plugins/Experimental/MotionTrailEditorMode/Source/MotionTrailEditorMode/Public/MotionTrailEditorMode.h
+++ b/Engine/Plugins/Experimental/MotionTrailEditorMode/Source/MotionTrailEditorMode/Public/MotionTrailEditorMode.h
@@ -87,7 +87,7 @@ public:
 	// UEdMode interface
 	virtual void Enter() override;
 	virtual void Exit() override;
-	void CreateToolkit() override;
+	virtual void CreateToolkit() override;
 	
 	virtual void Render(const FSceneView* View, FViewport* Viewport, FPrimitiveDrawInterface* PDI) override;
 	virtual void DrawHUD(FEditorViewportClient* ViewportClient, FViewport* Viewport, const FSceneView* View, FCanvas* Canvas) override;
diff --git a/Engine/Plugins/Experimental/MotionTrailEditorMode/Source/MotionTrailEditorMode/Public/MotionTrailEditorModeToolkit.h b/Engine/Plugins/Experimental/MotionTrailEditorMode/Source/MotionTrailEditorMode/Public/MotionTrailEditorModeToolkit.h
--- a/Engine/Plugins/Experimental/MotionTrailEditorMode/Source/MotionTrailEditorMode/Public/MotionTrailEditorModeToolkit.h
+++ b/Engine/Plugins/Experimental/MotionTrailEditorMode/Source/MotionTrailEditorMode/Public/MotionTrailEditorModeToolkit.h
@@ -21,12 +21,12 @@ public:
 	FMotionTrailEditorModeToolkit();
 	
 	/** FModeToolkit interface */
-	virtual void Init(const TSharedPtr<IToolkitHost>& InitToolkitHost) override;
+	virtual void Init(const TSharedPtr<IToolkitHost>& InitToolkitHost, TWeakObjectPtr<UEdMode> InOwningMode) override;
 
 	/** IToolkit interface */
 	virtual FName GetToolkitFName() const override;
 	virtual FText GetBaseToolkitName() const override;
-	TSharedPtr<SWidget> GetInlineContent() const override;
+	virtual TSharedPtr<SWidget> GetInlineContent() const override;
 
 	void SetTimingStats(const TArray<TMap<FString, FTimespan>>& HierarchyStats);
 
